Datasmith - Move(add) the DatasmithClotranslator to the Experimental folder.

#rnx

diff --git a/Engine/Plugins/Experimental/DatasmithCloTranslator/DatasmithCloTranslator.uplugin b/Engine/Plugins/Experimental/DatasmithCloTranslator/DatasmithCloTranslator.uplugin
new file mode 100644
--- /dev/null
+++ b/Engine/Plugins/Experimental/DatasmithCloTranslator/DatasmithCloTranslator.uplugin
@@ -0,0 +1,37 @@
+{
+	"FileVersion": 3,
+	"Version": 1,
+	"VersionName": "1.0",
+	"FriendlyName": "Datasmith Clo Translator",
+	"Description": "Datasmith translator for Clo assets",
+	"Category": "Importers",
+	"CreatedBy": "Epic Games, Inc.",
+	"CreatedByURL": "https://epicgames.com",
+	"DocsURL": "",
+	"MarketplaceURL": "",
+	"SupportURL": "",
+	"CanContainContent": false,
+	"IsExperimentalVersion": true,
+	"Installed": false,
+	"SupportedTargetPlatforms": [
+		"Win64"
+	],
+	"Modules": [
+		{
+			"Name": "DatasmithCloTranslator",
+			"Type": "Runtime",
+			"LoadingPhase": "Default"
+		}
+	],
+	"Plugins": [
+		{
+			"Name": "DatasmithImporter",
+			"Enabled": true
+		},
+		{
+			"Name": "ChaosClothAsset",
+			"Enabled": true
+		},
+	],
+	"IsBetaVersion": false
+}
\ No newline at end of file
diff --git a/Engine/Plugins/Experimental/DatasmithCloTranslator/Source/DatasmithCloTranslator/DatasmithCloTranslator.Build.cs b/Engine/Plugins/Experimental/DatasmithCloTranslator/Source/DatasmithCloTranslator/DatasmithCloTranslator.Build.cs
new file mode 100644
--- /dev/null
+++ b/Engine/Plugins/Experimental/DatasmithCloTranslator/Source/DatasmithCloTranslator/DatasmithCloTranslator.Build.cs
@@ -0,0 +1,28 @@
+// Copyright Epic Games, Inc. All Rights Reserved.
+
+using UnrealBuildTool;
+
+public class DatasmithCloTranslator : ModuleRules
+{
+	public DatasmithCloTranslator(ReadOnlyTargetRules Target) : base(Target)
+	{
+		PublicDependencyModuleNames.AddRange(
+			new string[]
+			{
+				"Core",
+			}
+		);
+
+		PrivateDependencyModuleNames.AddRange(
+			new string[]
+			{
+				"CoreUObject",
+				"DatasmithContent",
+				"DatasmithCore",
+				"DatasmithTranslator",
+				"Json",
+				"MeshDescription",
+			}
+		);
+	}
+}
diff --git a/Engine/Plugins/Experimental/DatasmithCloTranslator/Source/DatasmithCloTranslator/Private/DatasmithCloTranslator.cpp b/Engine/Plugins/Experimental/DatasmithCloTranslator/Source/DatasmithCloTranslator/Private/DatasmithCloTranslator.cpp
new file mode 100644
--- /dev/null
+++ b/Engine/Plugins/Experimental/DatasmithCloTranslator/Source/DatasmithCloTranslator/Private/DatasmithCloTranslator.cpp
@@ -0,0 +1,538 @@
+// Copyright Epic Games, Inc. All Rights Reserved.
+
+#include "DatasmithCloTranslator.h"
+
+#include "DatasmithCloth.h"
+#include "DatasmithMesh.h"
+#include "DatasmithSceneFactory.h"
+#include "DatasmithUtils.h"
+#include "IDatasmithSceneElements.h"
+
+#include "Algo/MinElement.h"
+#include "HAL/FileManager.h"
+#include "Serialization/JsonReader.h"
+#include "Serialization/JsonSerializer.h"
+
+
+DEFINE_LOG_CATEGORY_STATIC(LogCloTranslator, Log, All);
+
+
+struct FCloPanel
+{
+	uint32 Id = 0;
+	uint32 PhysicalPropertyId = 0;
+	TArray<FVector2f> Positions2D;
+	TArray<FVector3f> Positions3D;
+	TArray<uint32> TriangleIndices;
+};
+
+struct FCloSewingInfo
+{
+	uint32 Index = 0;
+	uint32 Seam0PanelIndex = 0;
+	uint32 Seam1PanelIndex = 0;
+	TArray<uint32> Seam0MeshIndices;
+	TArray<uint32> Seam1MeshIndices;
+};
+
+struct FCloPhysicalProperties
+{
+	uint32 Index;
+	TMap<FString, double> Values;
+};
+
+struct FCloCloth
+{
+	TArray<FCloPanel> Panels;
+	TArray<FCloSewingInfo> SewingInfos;
+	TArray<FCloPhysicalProperties> PhysicsProperties;
+};
+
+void FDatasmithCloTranslator::Initialize(FDatasmithTranslatorCapabilities& OutCapabilities)
+{
+	OutCapabilities.SupportedFileFormats.Emplace(TEXT("json"), TEXT("CLO json files"));
+}
+
+void FDatasmithCloTranslator::SetSceneImportOptions(const TArray<TObjectPtr<UDatasmithOptionsBase>>& AllOptions)
+{
+	for (const TObjectPtr<UDatasmithOptionsBase>& Option : AllOptions)
+	{
+		if (const UDatasmithCloImportOptions* CloOptions = Cast<UDatasmithCloImportOptions>(Option.Get()))
+		{
+			ImportOptions = CloOptions->Options;
+		}
+	}
+}
+
+void FDatasmithCloTranslator::GetSceneImportOptions(TArray<TObjectPtr<UDatasmithOptionsBase>>& AllOptions)
+{
+	const TObjectPtr<UDatasmithCloImportOptions> CloOptionsPtr = Datasmith::MakeOptionsObjectPtr<UDatasmithCloImportOptions>();
+	CloOptionsPtr->Options = ImportOptions;
+	AllOptions.Add(CloOptionsPtr);
+}
+
+bool FDatasmithCloTranslator::LoadScene(TSharedRef<IDatasmithScene> DatasmithScene)
+{
+	if (!ParseFromJson())
+	{
+		return false;
+	}
+	check(CloClothPtr);
+
+	// Creates a cloth asset element
+	const TSharedRef<IDatasmithClothElement> ClothElement = FDatasmithSceneFactory::CreateCloth(TEXT("cloth asset"));
+	DatasmithScene->AddCloth(ClothElement);
+
+	if (ImportOptions.bGenerateClothStaticMesh)
+	{
+		FString Name = ClothElement->GetName();
+		Name += TEXT("_dbgmesh");
+		const auto MeshElement = FDatasmithSceneFactory::CreateMesh(*Name);
+		DatasmithScene->AddMesh(MeshElement);
+
+		const auto MeshActor = FDatasmithSceneFactory::CreateMeshActor(*FString::Printf(TEXT("%s"), *Name));
+		DatasmithScene->AddActor(MeshActor);
+		MeshActor->SetStaticMeshPathName(MeshElement->GetName());
+	}
+
+	if (ImportOptions.bGeneratePerPanel2DMesh)
+	{
+		FCloCloth& CloCloth = *CloClothPtr;
+		for (const FCloPanel& Panel : CloCloth.Panels)
+		{
+			FString Name = TEXT("Panel_") + FString::FromInt(Panel.Id);
+			PanelMeshNames.Add(Name);
+			auto MeshElement = FDatasmithSceneFactory::CreateMesh(*Name);
+			DatasmithScene->AddMesh(MeshElement);
+
+			auto MeshActor = FDatasmithSceneFactory::CreateMeshActor(*Name);
+			DatasmithScene->AddActor(MeshActor);
+			MeshActor->SetStaticMeshPathName(MeshElement->GetName());
+		}
+	}
+
+	if (ImportOptions.bGenerateClothActor)
+	{
+		const TSharedRef<IDatasmithClothActorElement> ClothActorElement = FDatasmithSceneFactory::CreateClothActor(TEXT("cloth actor"));
+		ClothActorElement->SetCloth(ClothElement->GetName());
+		DatasmithScene->AddActor(ClothActorElement);
+	}
+
+	return true;
+}
+
+void FDatasmithCloTranslator::UnloadScene()
+{
+	CloClothPtr.Reset();
+}
+
+void AppendIntList(const TArray<TSharedPtr<FJsonValue>>& Src, TArray<uint32>& Dst)
+{
+	Dst.Reserve(Dst.Num() + Src.Num());
+	for (const TSharedPtr<FJsonValue>& IntJson : Src)
+	{
+		IntJson->TryGetNumber(Dst.AddDefaulted_GetRef());
+	}
+};
+
+#define ParseVerify(Cond, Text) \
+if (!(Cond)) { \
+	UE_LOG(LogCloTranslator, Warning, Text); \
+	return false; \
+}
+
+bool FDatasmithCloTranslator::LoadCloth(const TSharedRef<IDatasmithClothElement> ClothElement, FDatasmithClothElementPayload& OutClothPayload)
+{
+	if (!CloClothPtr.IsValid())
+	{
+		return false;
+	}
+
+	FCloCloth& CloCloth = *CloClothPtr;
+
+	FDatasmithCloth& DsCloth = OutClothPayload.Cloth;
+
+	for (FCloPanel& CloPanel : CloCloth.Panels)
+	{
+		FDatasmithClothPattern& Panel = DsCloth.Patterns.AddDefaulted_GetRef();
+		Panel.SimPosition = MoveTemp(CloPanel.Positions2D);
+		Panel.SimRestPosition = MoveTemp(CloPanel.Positions3D);
+		Panel.SimTriangleIndices = MoveTemp(CloPanel.TriangleIndices);
+
+		// Invert winding order
+		for (int32 TriIndex = 0; TriIndex < Panel.SimTriangleIndices.Num() / 3; ++TriIndex)
+		{
+			Swap(Panel.SimTriangleIndices[3 * TriIndex + 1], Panel.SimTriangleIndices[3 * TriIndex + 2]);
+		}
+
+		if (ImportOptions.Scale != 1.0f)
+		{
+			for (auto& Vertex : Panel.SimRestPosition)
+			{
+				Vertex *= ImportOptions.Scale;
+			}
+		}
+
+		constexpr int32 MaxInvalidTriangleCount = 5;
+		int32 InvalidTriangleCount = 0;
+		for (int32 TriIndex = 0; TriIndex < Panel.SimTriangleIndices.Num() / 3; ++TriIndex)
+		{
+			FVector3f& A = Panel.SimRestPosition[Panel.SimTriangleIndices[3 * TriIndex]];
+			FVector3f& B = Panel.SimRestPosition[Panel.SimTriangleIndices[3 * TriIndex + 1]];
+			FVector3f& C = Panel.SimRestPosition[Panel.SimTriangleIndices[3 * TriIndex + 2]];
+			FVector3f AB = B - A;
+			FVector3f AC = C - A;
+			FVector3f ABC = AB ^ AC;
+			const float SizeSquared = ABC.SizeSquared();
+			if (SizeSquared < SMALL_NUMBER)
+			{
+				InvalidTriangleCount++;
+				if (InvalidTriangleCount >= MaxInvalidTriangleCount)
+				{
+					break;
+				}
+				UE_LOG(LogCloTranslator, Warning, TEXT("Invalid triangle found. Import may fail. Panel %d, triangle indices %d, %d and %d."),
+					CloPanel.Id,
+					Panel.SimTriangleIndices[3 * TriIndex],
+					Panel.SimTriangleIndices[3 * TriIndex + 1],
+					Panel.SimTriangleIndices[3 * TriIndex + 2]);
+				UE_LOG(LogCloTranslator, Warning, TEXT("note: triangle (%s) (%s) (%s)"), *A.ToString(), *B.ToString(), *C.ToString());
+			}
+		}
+
+		if (InvalidTriangleCount >= MaxInvalidTriangleCount)
+		{
+			UE_LOG(LogCloTranslator, Warning, TEXT("Skip invalid panel %d: too many invalid triangles."), CloPanel.Id);
+			DsCloth.Patterns.Pop();
+			continue;
+		}
+
+		if (!Panel.IsValid())
+		{
+			UE_LOG(LogCloTranslator, Warning, TEXT("Skip invalid panel %d."), CloPanel.Id);
+			DsCloth.Patterns.Pop();
+			continue;
+		}
+	}
+
+	int PropertySetIndex = 0;
+	for (FCloPhysicalProperties& CloProperties : CloCloth.PhysicsProperties)
+	{
+		FDatasmithClothPresetPropertySet& Set = DsCloth.PropertySets.AddDefaulted_GetRef();
+		Set.SetName = FString::Printf(TEXT("CloPropertySet %d"), PropertySetIndex++);
+		for (const auto& [Name, Value] : CloProperties.Values)
+		{
+			Set.Properties.Add({FName(Name), Value});
+		}
+	}
+
+	TSet<uint32> ValidPanelIds;
+	for (FCloPanel& CloPanel : CloCloth.Panels)
+	{
+		ValidPanelIds.Add(CloPanel.Id);
+	}
+
+	for (const FCloSewingInfo& SI : CloCloth.SewingInfos)
+	{
+		if (ValidPanelIds.Contains(SI.Seam0PanelIndex) && ValidPanelIds.Contains(SI.Seam1PanelIndex))
+		{
+			FDatasmithClothSewingInfo DsSewingInfo;
+			DsSewingInfo.Seam0PanelIndex = SI.Seam0PanelIndex;
+			DsSewingInfo.Seam1PanelIndex = SI.Seam1PanelIndex;
+			DsSewingInfo.Seam0MeshIndices = SI.Seam0MeshIndices;
+			DsSewingInfo.Seam1MeshIndices = SI.Seam1MeshIndices;
+
+			DsCloth.Sewing.Add(DsSewingInfo);
+		}
+	}
+
+	if (DsCloth.Patterns.IsEmpty())
+	{
+		UE_LOG(LogCloTranslator, Warning, TEXT("Empty cloth translated: %s"), ClothElement->GetName());
+		return false;
+	}
+
+	return true;
+}
+
+
+bool FDatasmithCloTranslator::LoadStaticMesh(const TSharedRef<IDatasmithMeshElement> MeshElement, FDatasmithMeshElementPayload& OutMeshPayload)
+{
+	// This makes a static mesh out of the cloth data in order to have a visualization. Debug purpose only.
+	if (!CloClothPtr.IsValid())
+	{
+		return false;
+	}
+
+	FCloCloth& CloCloth = *CloClothPtr;
+
+	if (ImportOptions.bGeneratePerPanel2DMesh)
+	{
+		int32 PanelIndex = PanelMeshNames.IndexOfByKey(MeshElement->GetName());
+		if (PanelIndex != INDEX_NONE )
+		{
+			if (!ensure(CloCloth.Panels.IsValidIndex(PanelIndex)))
+			{
+				return false;
+			}
+
+			FCloPanel& Panel = CloCloth.Panels[PanelIndex];
+
+			FDatasmithMesh DsMesh;
+			DsMesh.SetVerticesCount(Panel.Positions2D.Num());
+			DsMesh.AddUVChannel();
+			DsMesh.SetUVCount(0, Panel.Positions2D.Num());
+			DsMesh.SetFacesCount(Panel.TriangleIndices.Num() / 3);
+
+			int32 VertexIndex = 0;
+			for (FVector2f Vertex : Panel.Positions2D)
+			{
+				Vertex *= ImportOptions.Scale;
+				DsMesh.SetVertex(VertexIndex++, Vertex.X * 0.1f, Vertex.Y * 0.1f, 0);
+			}
+
+			FBox2f Box;
+			for (const auto& Vertex : Panel.Positions2D)
+			{
+				Box += Vertex;
+			}
+
+			FVector2f MinimumUV;
+			if (auto* MEy = Algo::MinElement(Panel.Positions2D, [](auto UVa, auto UVb){ return UVa.Y < UVb.Y;}))
+			{
+				auto* MEx = Algo::MinElement(Panel.Positions2D, [](auto UVa, auto UVb){ return UVa.X < UVb.X;});
+				MinimumUV.X = MEx->X;
+				MinimumUV.Y = MEy->Y;
+			}
+
+			int32 UVIndex = 0;
+			for (FVector2f UV : Panel.Positions2D)
+			{
+				UV -= MinimumUV; // offset so that the range start to 0, 0
+				UV *= 0.001; // as the data is in mm, maps a ~1 meter pattern in the 0, 1 range
+				DsMesh.SetUV(0, UVIndex++, UV.X, UV.Y);
+			}
+
+			for (int32 PanelFaceIndex = 0; PanelFaceIndex < Panel.TriangleIndices.Num() / 3; ++PanelFaceIndex)
+			{
+				DsMesh.SetFace(PanelFaceIndex,
+					Panel.TriangleIndices[3*PanelFaceIndex+0],
+					Panel.TriangleIndices[3*PanelFaceIndex+2],
+					Panel.TriangleIndices[3*PanelFaceIndex+1],
+					ImportOptions.bPerPanelMaterial ? Panel.Id : 0);
+				DsMesh.SetFaceUV(PanelFaceIndex, 0,
+					Panel.TriangleIndices[3*PanelFaceIndex+0],
+					Panel.TriangleIndices[3*PanelFaceIndex+2],
+					Panel.TriangleIndices[3*PanelFaceIndex+1]);
+				DsMesh.SetFaceSmoothingMask(PanelFaceIndex, 1);
+			}
+
+			FMeshDescription& MeshDescription = OutMeshPayload.LodMeshes.AddDefaulted_GetRef();
+			FDatasmithMeshUtils::ToMeshDescription(DsMesh, MeshDescription);
+			return true;
+		}
+	}
+
+	FDatasmithMesh DsMesh;
+	DsMesh.AddUVChannel();
+
+	int32 VerticesCount = 0;
+	int32 UVCount = 0;
+	int32 FacesCount = 0;
+
+	for (const FCloPanel& Panel : CloCloth.Panels)
+	{
+		VerticesCount += Panel.Positions3D.Num();
+		UVCount += Panel.Positions2D.Num();
+		FacesCount += Panel.TriangleIndices.Num() / 3;
+	}
+
+	DsMesh.SetVerticesCount(VerticesCount);
+	DsMesh.SetUVCount(0, UVCount);
+	DsMesh.SetFacesCount(FacesCount);
+
+	int32 VertexIndex = 0;
+	int32 UVIndex = 0;
+	int32 FaceIndexOffset = 0;
+	int32 TriangleIndicesOffset = 0;
+	for (const FCloPanel& Panel : CloCloth.Panels)
+	{
+		for (FVector3f Vertex : Panel.Positions3D)
+		{
+			Vertex *= ImportOptions.Scale;
+			DsMesh.SetVertex(VertexIndex++, Vertex.X, Vertex.Y, Vertex.Z);
+		}
+
+		FVector2f MinimumUV;
+		if (auto* MEy = Algo::MinElement(Panel.Positions2D, [](auto UVa, auto UVb){ return UVa.Y < UVb.Y;}))
+		{
+			auto* MEx = Algo::MinElement(Panel.Positions2D, [](auto UVa, auto UVb){ return UVa.X < UVb.X;});
+			MinimumUV.X = MEx->X;
+			MinimumUV.Y = MEy->Y;
+		}
+
+		for (FVector2f UV : Panel.Positions2D)
+		{
+			UV -= MinimumUV; // offset so that the range start to 0, 0
+			UV *= 0.001; // as the data is in mm, maps a ~1 meter pattern in the 0, 1 range
+			DsMesh.SetUV(0, UVIndex++, UV.X, UV.Y);
+		}
+
+		for (int32 PanelFaceIndex = 0; PanelFaceIndex < Panel.TriangleIndices.Num() / 3; ++PanelFaceIndex)
+		{
+			DsMesh.SetFace(FaceIndexOffset + PanelFaceIndex,
+				TriangleIndicesOffset + Panel.TriangleIndices[3*PanelFaceIndex+0],
+				TriangleIndicesOffset + Panel.TriangleIndices[3*PanelFaceIndex+2],
+				TriangleIndicesOffset + Panel.TriangleIndices[3*PanelFaceIndex+1],
+				ImportOptions.bPerPanelMaterial ? Panel.Id : 0);
+			DsMesh.SetFaceUV(FaceIndexOffset + PanelFaceIndex, 0,
+				TriangleIndicesOffset + Panel.TriangleIndices[3*PanelFaceIndex+0],
+				TriangleIndicesOffset + Panel.TriangleIndices[3*PanelFaceIndex+2],
+				TriangleIndicesOffset + Panel.TriangleIndices[3*PanelFaceIndex+1]);
+			DsMesh.SetFaceSmoothingMask(FaceIndexOffset + PanelFaceIndex, 1);
+		}
+
+		FaceIndexOffset += Panel.TriangleIndices.Num() / 3;
+		TriangleIndicesOffset += Panel.Positions3D.Num();
+	}
+
+	FMeshDescription& MeshDescription = OutMeshPayload.LodMeshes.AddDefaulted_GetRef();
+	FDatasmithMeshUtils::ToMeshDescription(DsMesh, MeshDescription);
+
+	return true;
+}
+
+bool FDatasmithCloTranslator::ParseFromJson()
+{
+	FString JsonPath = GetSource().GetSourceFile();
+	TUniquePtr<FArchive> FileReader(IFileManager::Get().CreateFileReader(*JsonPath));
+	if (!FileReader)
+	{
+		return false;
+	}
+
+	TSharedRef Reader(TJsonReaderFactory<char>::Create(FileReader.Get()));
+	TSharedPtr<class FJsonObject> RootJsonObject;
+	FJsonSerializer::Deserialize(Reader, RootJsonObject);
+	if (!RootJsonObject.IsValid())
+	{
+		return false;
+	}
+
+	FCloCloth CloCloth;
+	ParseVerify(RootJsonObject, TEXT("Bad root object"));
+
+	int version = RootJsonObject->HasTypedField<EJson::Array>(TEXT("PanelInformation")) ? 1 : 0;
+	// 0: MD AdditionalInformation
+	// 1: MD BendingwingInformation
+
+	const TCHAR* kPanelInfoName = version == 1 ? TEXT("PanelInformation") : TEXT("PanelInformation");
+	const TCHAR* kTriangleIndexName = version == 1 ? TEXT("panelTriangleIndex") : TEXT("triangleIndex");
+
+	ParseVerify(RootJsonObject->HasTypedField<EJson::Array>(kPanelInfoName), TEXT("Missing Panel Information list"));
+// 	ParseVerify(RootJsonObject->HasTypedField<EJson::Array>(TEXT("SewingInformation")), TEXT("Missing SewingInformation list"));
+
+	{
+		const TArray<TSharedPtr<FJsonValue>>& PanelInformationsJsonObject = RootJsonObject->GetArrayField(kPanelInfoName);
+		CloCloth.Panels.Reserve(PanelInformationsJsonObject.Num());
+
+		for (const TSharedPtr<FJsonValue>& PanelJsonValue : PanelInformationsJsonObject)
+		{
+			const TSharedPtr<FJsonObject>& PanelJsonObject = PanelJsonValue->AsObject();
+			if (PanelJsonObject->TryGetField(TEXT("skip")))
+			{
+				continue;
+			}
+			FCloPanel& Panel = CloCloth.Panels.AddDefaulted_GetRef();
+			Panel.Id = PanelJsonObject->GetIntegerField( version == 1 ? TEXT("panelId") : TEXT("index"));
+			Panel.PhysicalPropertyId = PanelJsonObject->GetIntegerField(TEXT("physicalPropertyIndex"));
+			if (ImportOptions.bSkipPanelsWithoutPhysicsProperties && Panel.PhysicalPropertyId == -1)
+			{
+				CloCloth.Panels.Pop();
+				continue;
+			}
+
+			ParseVerify(PanelJsonObject->HasTypedField<EJson::Array>(TEXT("2DPositions")), TEXT("Missing 2DPositions list"));
+			ParseVerify(PanelJsonObject->HasTypedField<EJson::Array>(TEXT("3DPositions")), TEXT("Missing 3DPositions list"));
+			ParseVerify(PanelJsonObject->HasTypedField<EJson::Array>(kTriangleIndexName), TEXT("Missing triangle Index list"));
+
+			{
+				const TArray<TSharedPtr<FJsonValue>>& Position2dsJson = PanelJsonObject->GetArrayField(TEXT("2DPositions"));
+				Panel.Positions2D.Reserve(Position2dsJson.Num());
+				for (const TSharedPtr<FJsonValue>& Position2d : Position2dsJson)
+				{
+					const TArray<TSharedPtr<FJsonValue>>& XY = Position2d->AsArray();
+					ParseVerify(XY.Num() == 2, TEXT("2DPositions issue"));
+
+					FVector2f& Position = Panel.Positions2D.AddDefaulted_GetRef();
+					XY[0]->TryGetNumber(Position.X);
+					XY[1]->TryGetNumber(Position.Y);
+				}
+			}
+
+			{
+				const TArray<TSharedPtr<FJsonValue>>& Position3dsJson = PanelJsonObject->GetArrayField(TEXT("3DPositions"));
+				Panel.Positions3D.Reserve(Position3dsJson.Num());
+				for (const TSharedPtr<FJsonValue>& Position3d : Position3dsJson)
+				{
+					const TArray<TSharedPtr<FJsonValue>>& XYZ = Position3d->AsArray();
+					ParseVerify(XYZ.Num() == 3, TEXT("3DPositions issue"));
+
+					// inputs are in millimeters, UE in centimeters. Also, we convert to our Z-up frame
+					FVector3f& Position = Panel.Positions3D.AddDefaulted_GetRef();
+					XYZ[2]->TryGetNumber(Position.X);
+					XYZ[0]->TryGetNumber(Position.Y);
+					XYZ[1]->TryGetNumber(Position.Z);
+					Position.Y *= -1;
+					Position.X *= -1;
+					Position *= 0.1;
+				}
+			}
+
+			AppendIntList(PanelJsonObject->GetArrayField(kTriangleIndexName), Panel.TriangleIndices);
+			ParseVerify(Panel.TriangleIndices.Num() % 3 == 0, TEXT("Triangle indices issue"));
+			ParseVerify(Panel.Positions3D.Num() == Panel.Positions2D.Num(), TEXT("Positions 2d/3d count differ"));
+		}
+	}
+
+	const TArray<TSharedPtr<FJsonValue>>* SewingInformationPtr;
+	if (RootJsonObject->TryGetArrayField(TEXT("SewingInformation"), SewingInformationPtr))
+	{
+		auto& SewingInformation = *SewingInformationPtr;
+		CloCloth.SewingInfos.Reserve(SewingInformation.Num());
+
+		for (const TSharedPtr<FJsonValue>& SewingInfoJsonValue : SewingInformation)
+		{
+			const TSharedPtr<FJsonObject>& SewingInfoJsonObject = SewingInfoJsonValue->AsObject();
+			ParseVerify(SewingInfoJsonObject->HasTypedField<EJson::Array>(TEXT("seam0MeshIndexList")), TEXT("Missing seam0MeshIndexList list"));
+			ParseVerify(SewingInfoJsonObject->HasTypedField<EJson::Array>(TEXT("seam1MeshIndexList")), TEXT("Missing seam1MeshIndexList list"));
+
+			FCloSewingInfo& SewInfo = CloCloth.SewingInfos.AddDefaulted_GetRef();
+			SewInfo.Index = SewingInfoJsonObject->GetIntegerField(TEXT("index"));
+			SewInfo.Seam0PanelIndex = SewingInfoJsonObject->GetIntegerField(TEXT("seam0PanelIndex"));
+			SewInfo.Seam1PanelIndex = SewingInfoJsonObject->GetIntegerField(TEXT("seam1PanelIndex"));
+			AppendIntList(SewingInfoJsonObject->GetArrayField(TEXT("seam0MeshIndexList")), SewInfo.Seam0MeshIndices);
+			AppendIntList(SewingInfoJsonObject->GetArrayField(TEXT("seam1MeshIndexList")), SewInfo.Seam1MeshIndices);
+		}
+	}
+
+	{
+		const TArray<TSharedPtr<FJsonValue>>& PhysicalPropertiesJsonValues = RootJsonObject->GetArrayField(TEXT("physicalProperties"));
+		CloCloth.PhysicsProperties.Reserve(PhysicalPropertiesJsonValues.Num());
+
+		for (const TSharedPtr<FJsonValue>& PhysicalPropertiesJsonValue : PhysicalPropertiesJsonValues)
+		{
+			FCloPhysicalProperties& PhyProp = CloCloth.PhysicsProperties.AddDefaulted_GetRef();
+			const TSharedPtr<FJsonObject>& PhysicalPropertiesJsonObject = PhysicalPropertiesJsonValue->AsObject();
+			for (const auto& [Name, JsonValue] : PhysicalPropertiesJsonObject->Values)
+			{
+				PhyProp.Values.Add(Name, JsonValue->AsNumber());
+			}
+		}
+	}
+
+	CloClothPtr = MakeShared<FCloCloth>();
+	*CloClothPtr = MoveTemp(CloCloth);
+	return true;
+}
+
diff --git a/Engine/Plugins/Experimental/DatasmithCloTranslator/Source/DatasmithCloTranslator/Private/DatasmithCloTranslator.h b/Engine/Plugins/Experimental/DatasmithCloTranslator/Source/DatasmithCloTranslator/Private/DatasmithCloTranslator.h
new file mode 100644
--- /dev/null
+++ b/Engine/Plugins/Experimental/DatasmithCloTranslator/Source/DatasmithCloTranslator/Private/DatasmithCloTranslator.h
@@ -0,0 +1,33 @@
+// Copyright Epic Games, Inc. All Rights Reserved.
+
+#pragma once
+
+#include "DatasmithTranslator.h"
+#include "DatasmithCloImportOptions.h"
+
+
+
+class DATASMITHCLOTRANSLATOR_API FDatasmithCloTranslator : public IDatasmithTranslator
+{
+public:
+	virtual FName GetFName() const override { return "DatasmithCloTranslator"; };
+	virtual void Initialize(FDatasmithTranslatorCapabilities& OutCapabilities) override;
+
+	virtual void SetSceneImportOptions(const TArray<TObjectPtr<UDatasmithOptionsBase>>& AllOptions) override;
+	virtual void GetSceneImportOptions(TArray<TObjectPtr<UDatasmithOptionsBase>>& AllOptions) override;
+
+	virtual bool LoadScene(TSharedRef<IDatasmithScene> OutScene) override;
+	virtual void UnloadScene() override;
+
+	virtual bool LoadCloth(const TSharedRef<IDatasmithClothElement> ClothElement, FDatasmithClothElementPayload& OutClothPayload) override;
+	virtual bool LoadStaticMesh(const TSharedRef<IDatasmithMeshElement> MeshElement, FDatasmithMeshElementPayload& OutMeshPayload) override;
+
+private:
+	bool ParseFromJson();
+
+private:
+	TSharedPtr<struct FCloCloth> CloClothPtr;
+	FCloOptions ImportOptions;
+	TArray<FString> PanelMeshNames;
+};
+
diff --git a/Engine/Plugins/Experimental/DatasmithCloTranslator/Source/DatasmithCloTranslator/Private/DatasmithCloTranslatorModule.cpp b/Engine/Plugins/Experimental/DatasmithCloTranslator/Source/DatasmithCloTranslator/Private/DatasmithCloTranslatorModule.cpp
new file mode 100644
--- /dev/null
+++ b/Engine/Plugins/Experimental/DatasmithCloTranslator/Source/DatasmithCloTranslator/Private/DatasmithCloTranslatorModule.cpp
@@ -0,0 +1,22 @@
+// Copyright Epic Games, Inc. All Rights Reserved.
+
+#include "DatasmithCloTranslatorModule.h"
+#include "DatasmithCloTranslator.h"
+#include "DatasmithTranslator.h"
+
+
+#define LOCTEXT_NAMESPACE "FDatasmithCloTranslatorModule"
+
+void FDatasmithCloTranslatorModule::StartupModule()
+{
+	Datasmith::RegisterTranslator<FDatasmithCloTranslator>();
+}
+
+void FDatasmithCloTranslatorModule::ShutdownModule()
+{
+	Datasmith::UnregisterTranslator<FDatasmithCloTranslator>();
+}
+
+#undef LOCTEXT_NAMESPACE
+
+IMPLEMENT_MODULE(FDatasmithCloTranslatorModule, DatasmithCloTranslator)
diff --git a/Engine/Plugins/Experimental/DatasmithCloTranslator/Source/DatasmithCloTranslator/Public/DatasmithCloImportOptions.h b/Engine/Plugins/Experimental/DatasmithCloTranslator/Source/DatasmithCloTranslator/Public/DatasmithCloImportOptions.h
new file mode 100644
--- /dev/null
+++ b/Engine/Plugins/Experimental/DatasmithCloTranslator/Source/DatasmithCloTranslator/Public/DatasmithCloImportOptions.h
@@ -0,0 +1,61 @@
+// Copyright Epic Games, Inc. All Rights Reserved.
+
+#pragma once
+
+#include "DatasmithImportOptions.h"
+
+#include "DatasmithCloImportOptions.generated.h"
+
+USTRUCT(BlueprintType)
+struct FCloOptions
+{
+	GENERATED_BODY()
+
+	UPROPERTY(config, EditAnywhere, BlueprintReadWrite, Category = "Clo")
+	bool bGenerateClothStaticMesh = false;
+
+	UPROPERTY(config, EditAnywhere, BlueprintReadWrite, Category = "Clo")
+	bool bGeneratePerPanel2DMesh = false;
+
+	UPROPERTY(config, EditAnywhere, BlueprintReadWrite, Category = "Clo")
+	bool bGenerateClothActor = true;
+
+	UPROPERTY(config, EditAnywhere, BlueprintReadWrite, Category = "Clo")
+	bool bPerPanelMaterial = false;
+
+	UPROPERTY(config, EditAnywhere, BlueprintReadWrite, Category = "Clo")
+	bool bSkipPanelsWithoutPhysicsProperties = true;
+
+	UPROPERTY(config, EditAnywhere, BlueprintReadWrite, Category = "Clo")
+	float Scale = 1.0f;
+};
+
+
+UCLASS(BlueprintType, config = EditorPerProjectUserSettings)
+class DATASMITHCLOTRANSLATOR_API UDatasmithCloImportOptions : public UDatasmithOptionsBase
+{
+	GENERATED_BODY()
+
+public:
+	UPROPERTY(config, EditAnywhere, BlueprintReadWrite, Category = "Clo", meta = (ShowOnlyInnerProperties))
+	FCloOptions Options;
+
+#if WITH_EDITOR
+	bool CanEditChange(const FProperty* InProperty) const override
+	{
+		if (!Super::CanEditChange(InProperty))
+		{
+			return false;
+		}
+
+		FName PropertyFName = InProperty->GetFName();
+		if (PropertyFName == GET_MEMBER_NAME_CHECKED(FCloOptions, bPerPanelMaterial))
+		{
+			return Options.bGenerateClothStaticMesh;
+		}
+
+		return true;
+	}
+#endif //WITH_EDITOR
+};
+
diff --git a/Engine/Plugins/Experimental/DatasmithCloTranslator/Source/DatasmithCloTranslator/Public/DatasmithCloTranslatorModule.h b/Engine/Plugins/Experimental/DatasmithCloTranslator/Source/DatasmithCloTranslator/Public/DatasmithCloTranslatorModule.h
new file mode 100644
--- /dev/null
+++ b/Engine/Plugins/Experimental/DatasmithCloTranslator/Source/DatasmithCloTranslator/Public/DatasmithCloTranslatorModule.h
@@ -0,0 +1,13 @@
+// Copyright Epic Games, Inc. All Rights Reserved.
+
+#pragma once
+
+#include "CoreMinimal.h"
+#include "Modules/ModuleManager.h"
+
+class FDatasmithCloTranslatorModule : public IModuleInterface
+{
+public:
+	virtual void StartupModule() override;
+	virtual void ShutdownModule() override;
+};
